#PUMA Creating an exlicit dependency in the Editor build for the Runtime; fixes a Linux, Mac linking error

diff --git a/Engine/Plugins/Animation/PUMA/Source/Editor/PUMAEditor.Build.cs b/Engine/Plugins/Animation/PUMA/Source/Editor/PUMAEditor.Build.cs
--- a/Engine/Plugins/Animation/PUMA/Source/Editor/PUMAEditor.Build.cs
+++ b/Engine/Plugins/Animation/PUMA/Source/Editor/PUMAEditor.Build.cs
@@ -31,6 +31,7 @@ namespace UnrealBuildTool.Rules
                         "UnrealEd",
                         "Kismet",
                         "AnimGraph",
+						"AnimGraphRuntime",
                         "BlueprintGraph",
                     }
                 );
