#PUMA ContentExamples changes

diff --git a/Engine/Plugins/Animation/PUMA/Source/Runtime/Private/AnimInstances/AnimInstance_DistanceMatching.cpp b/Engine/Plugins/Animation/PUMA/Source/Runtime/Private/AnimInstances/AnimInstance_DistanceMatching.cpp
new file mode 100644
--- /dev/null
+++ b/Engine/Plugins/Animation/PUMA/Source/Runtime/Private/AnimInstances/AnimInstance_DistanceMatching.cpp
@@ -0,0 +1,335 @@
+// Copyright Epic Games, Inc. All Rights Reserved.
+
+
+#include "AnimInstances/AnimInstance_DistanceMatching.h"
+#include "Animation/AnimSequence.h"
+
+/////////////////////////////////////////////////////
+// UAnimInstance_Locomotion
+/////////////////////////////////////////////////////
+
+DEFINE_LOG_CATEGORY_STATIC(LogLocomotion, Log, All);
+
+UAnimInstance_DistanceMatching::UAnimInstance_DistanceMatching(const FObjectInitializer& ObjectInitializer)
+	: Super(ObjectInitializer)
+{
+}
+
+void UAnimInstance_DistanceMatching::NativeBeginPlay()
+{
+	Super::NativeBeginPlay();
+
+	IKFootRootTransformLastFrame = FTransform::Identity;
+}
+
+void UAnimInstance_DistanceMatching::NativeUpdateAnimation(float DeltaSeconds)
+{
+	Super::NativeUpdateAnimation(DeltaSeconds);
+
+	ACharacter const * const CharacterOwner = Cast<ACharacter>(TryGetPawnOwner());
+	if (CharacterOwner != nullptr)
+	{
+		UCharacterMovementComponent const * const CharacterMovement = CharacterOwner->GetCharacterMovement();
+		if (CharacterMovement != nullptr)
+		{
+			UpdateCharacterLocomotionProperties(*CharacterOwner, *CharacterMovement, DeltaSeconds);
+			UpdateStartTransition(*CharacterOwner, *CharacterMovement, DeltaSeconds);
+			UpdateStopTransition(*CharacterOwner, *CharacterMovement, DeltaSeconds);
+		}
+	}	
+}
+
+void UAnimInstance_DistanceMatching::NativePostEvaluateAnimation()
+{
+	Super::NativePostEvaluateAnimation();
+
+	ACharacter const * const CharacterOwner = Cast<ACharacter>(TryGetPawnOwner());
+	if (CharacterOwner != nullptr)
+	{
+		IKFootRootTransformLastFrame = CharacterOwner->GetMesh()->GetSocketTransform("ik_foot_root", RTS_World);
+	}
+}
+
+void UAnimInstance_DistanceMatching::UpdateCharacterLocomotionProperties(const ACharacter& CharacterOwner, const UCharacterMovementComponent& CharacterMovement, float DeltaSeconds)
+{
+	MaxCharacterSpeed = CharacterMovement.MaxWalkSpeed;
+
+	CharacterVelocity = CharacterOwner.GetVelocity();
+	CharacterSpeed = CharacterVelocity.Size();
+	CharacterSpeed2D = CharacterVelocity.Size2D();
+	CharacterSpeedZ = CharacterVelocity.Z;
+
+	// Let Acceleration be instantaneous	
+	CharacterAcceleration = CharacterMovement.GetCurrentAcceleration();
+	CharacterAccelerationMag = CharacterAcceleration.Size();
+	CharacterAccelerationMag2D = CharacterAcceleration.Size2D();
+	CharacterAccelerationMagZ = CharacterAcceleration.Z;
+
+	// Initialize Rate & Warping values
+	PlayRateValue = 1.f;
+	StrideWarpingValue = 1.f;
+
+	float OverallSpeedRatio = 1.f;
+	float WalkSpeedRatio = 0.f;
+	float JogSpeedRatio = 0.f;
+	float SprintSpeedRatio = 0.f;
+
+	// ?????
+	// For now, just jog
+// 	if (WalkStateData.IsValid(*this) && JogStateData.IsValid(*this) && SprintStateData.IsValid(*this)
+// 		&& (WalkStateData.IsRelevant(*this) || JogStateData.IsRelevant(*this) || SprintStateData.IsRelevant(*this)))
+// 	{
+// 		WalkSpeedRatio = (CharacterSpeed / WalkSpeedIdeal) * WalkStateData.GetWeight(*this);
+// 		JogSpeedRatio = (CharacterSpeed / JogSpeedIdeal) * JogStateData.GetWeight(*this);
+// 		SprintSpeedRatio = (CharacterSpeed / SprintSpeedIdeal) * SprintStateData.GetWeight(*this);
+// 		OverallSpeedRatio = WalkSpeedRatio + JogSpeedRatio + SprintSpeedRatio;
+// 		CalculateStrideWarpingValues(PlayRateStrideWarpAlpha, OverallSpeedRatio);
+// 		//UE_LOG(LogLocomotion, Log, TEXT("OverallSpeedRatio: %f, WalkSpeedRatio: %f, JogSpeedRatio: %f, SprintSpeedRatio: %f"), OverallSpeedRatio, WalkSpeedRatio, JogSpeedRatio, SprintSpeedRatio);
+// 	}
+
+	if (JogStateData.IsValid(*this) && JogStateData.IsRelevant(*this))
+	{
+		OverallSpeedRatio = (CharacterSpeed / JogSpeedIdeal);// *JogStateData.GetWeight(*this);
+		CalculateStrideWarpingValues(PlayRateStrideWarpAlpha, OverallSpeedRatio);
+		//UE_LOG(LogLocomotion, Log, TEXT("OverallSpeedRatio: %f, WalkSpeedRatio: %f, JogSpeedRatio: %f, SprintSpeedRatio: %f"), OverallSpeedRatio, WalkSpeedRatio, JogSpeedRatio, SprintSpeedRatio);
+		UE_LOG(LogLocomotion, Log, TEXT("OverallSpeedRatio: %f"), OverallSpeedRatio);
+	}
+
+	// Slope Warping
+	FVector SlopeNormal = IKFootRootTransformLastFrame.GetUnitAxis(EAxis::Z);
+	FVector CharacterFacing2D = CharacterMovement.GetLastUpdateRotation().Vector();
+	float SlopeAngleGoal = FMath::Clamp(FMath::RadiansToDegrees(FMath::Acos((SlopeNormal | CharacterFacing2D))) - 90.f, -90.f, 90.f);
+	SlopeWarpingAdjustmentAlpha = (1.f - FMath::Clamp((CharacterSpeed / MaxCharacterSpeedForSlopeAngleAdjustmentRange), 0.f, 1.f));
+	float SlopeAngleAdjustmentSpeed = SlopeWarpingAdjustmentAlpha * (MaxSlopeAngleAdjustmentSpeed - MinSlopeAngleAdjustmentSpeed) + MinSlopeAngleAdjustmentSpeed;
+	SlopeAngle = FMath::FInterpTo(SlopeAngle, SlopeAngleGoal, DeltaSeconds, SlopeAngleAdjustmentSpeed);
+	if (MinSlopeWarpingAlpha <= MaxSlopeWarpingAlpha)
+	{
+		SlopeWarpingAlpha = FMath::Clamp(1.f - (CharacterSpeed / MaxCharacterSpeedForSlopeWarpingDamping), MinSlopeWarpingAlpha, MaxSlopeWarpingAlpha);
+	}
+	else
+	{
+		SlopeWarpingAlpha = 1.f;
+	}
+	//UE_LOG(LogLocomotion, Log, TEXT("SlopeAngle: %f"), SlopeAngle);
+	//UE_LOG(LogLocomotion, Log, TEXT("SlopeAngleGoal: %f"), SlopeAngleGoal);
+	//UE_LOG(LogLocomotion, Log, TEXT("SlopeAngleAdjustmentSpeed: %f"), SlopeAngleAdjustmentSpeed);
+	//UE_LOG(LogLocomotion, Log, TEXT("SlopeAlpha: %f"), SlopeWarpingAlpha);
+}
+
+
+void UAnimInstance_DistanceMatching::UpdateStartTransition(const ACharacter& CharacterOwner, const UCharacterMovementComponent& CharacterMovement, float DeltaTime)
+{
+	if (StartStateData.IsValid(*this))
+	{
+		// From CharacterMovementComp
+		//	ApplyVelocityBraking(DeltaTime, ActualBrakingFriction, BrakingDeceleration) = Velocity + ((-Friction) * Velocity + RevAccel) * dt;
+
+		//////const bool bIsStartStateRelevant = StartStateData.IsRelevant(*this);
+		const bool bIsStartStateRelevant = StartStateData.IsActiveState(*this);
+		const bool bIsStartStateFullWeight = StartStateData.IsFullWeight(*this);
+
+		if (!bIsStartStateRelevant || (bStartComplete && !bIsStartStateFullWeight))
+		{
+			bStartTransitionTriggered = false;
+			bStartComplete = false;
+			bPlayStart = false;
+			bWalkToStart = true;
+			//StartAnimPosition = 0.5f;
+			return;
+		}
+
+		FVector CharacterLocation = CharacterOwner.GetActorLocation();
+
+		FVector CharacterMovementDirection = CharacterVelocity.GetSafeNormal();
+
+		float Deceleration = CharacterAccelerationMag - CharacterMovement.GetMaxBrakingDeceleration() - (CharacterSpeed * CharacterMovement.GroundFriction);
+		float PredictedDistanceCovered = (-(CharacterSpeed * CharacterSpeed) / (2.f * Deceleration));
+
+		if (!bStartTransitionTriggered && bIsStartStateRelevant)
+		{
+			StartActualDistanceFromMarker = 0.0f;
+			StartAnimDistanceTraveled = 0.f;
+			StartAnimTimeElapsed = 0.f;
+			StartAnimLocation = CharacterLocation;
+
+			//*****
+			//UE_LOG(LogLocomotion, Log, TEXT("Character Acceleration: %f, Character Speed: %f"), CharacterAccelerationMag, CharacterSpeed);
+			//DrawDebugSphere(GetSkelMeshComponent()->GetWorld(), StartAnimLocation, 8.f, 16, FColor::Yellow);
+
+			StartAnimPosition = 0.5f;
+			bWalkToStart = (CharacterAccelerationMag < CharacterMovement.GetMaxAcceleration() * 0.45f) || CharacterMovement.GetMaxSpeed() < JogSpeedIdeal;
+
+			// ?????
+			// For now without walk anims
+			//StartAnimation = bWalkToStart ? WalkStart : JogStart;
+			StartAnimation = JogStart;
+			
+			bPlayStart = true;
+			bStartTransitionTriggered = true;
+
+			return;
+		}
+
+
+		// Predict Start?
+
+		// Find best place
+
+		// Adjust animation
+
+		if (StartAnimation && bPlayStart)
+		{
+			if (DistanceCurve.IsValid(StartAnimation))
+			{
+				StartAnimTimeElapsed += DeltaTime;
+				FVector FrameMovement = (CharacterVelocity * DeltaTime) + (CharacterAcceleration * DeltaTime * DeltaTime);
+									
+				float DesiredStartAnimDistance = DistanceCurve.GetValueAtPosition(StartAnimation, StartAnimPosition) + FrameMovement.Size2D();
+				float DesiredStartAnimPosition = FMath::Clamp(DistanceCurve.GetAnimPositionFromDistance(StartAnimation, DesiredStartAnimDistance), StartAnimPosition, StartAnimation->SequenceLength);
+				float DesiredStartAnimOverallRate = FMath::Clamp((DesiredStartAnimPosition - StartAnimPosition) / DeltaTime, 0.75f, 1.5f);
+				CalculateStrideWarpingValues(PlayRateStrideWarpAlpha, DesiredStartAnimOverallRate);
+				StartAnimPosition += DeltaTime * PlayRateValue;
+
+				//*****
+				//UE_LOG(LogLocomotion, Log, TEXT("DesiredStartAnimDistance: %f, DesiredStartAnimPosition: %f, DesiredStartAnimRate: %f"), DesiredStartAnimDistance, DesiredStartAnimPosition, DesiredStartAnimRate);
+				//UE_LOG(LogLocomotion, Log, TEXT("AnimPosition: %f"), StartAnimPosition);
+
+				//if (StartAnimPosition >= StartAnimation->SequenceLength - StartTransitionData.GetCrossfadeDuration(*this))
+				if (StartAnimPosition >= 2.f)
+				{
+					bStartComplete = true;
+				}
+			}
+			else // No DistanceCurve
+			{
+				StartAnimTimeElapsed += DeltaTime;
+				StartAnimPosition = DistanceCurve.GetValueAtPosition(StartAnimation, StartAnimTimeElapsed);
+			}
+		}
+
+	}
+}
+
+void UAnimInstance_DistanceMatching::UpdateStopTransition(const ACharacter& CharacterOwner, const UCharacterMovementComponent& CharacterMovement, float DeltaTime)
+{
+	if (StopStateData.IsValid(*this))
+	{
+		// From CharacterMovementComp
+		//	ApplyVelocityBraking(DeltaTime, ActualBrakingFriction, BrakingDeceleration) = Velocity + ((-Friction) * Velocity + RevAccel) * dt;
+
+		//const bool bIsStopStateRelevant = StopStateData.IsRelevant(*this);
+		const bool bIsStopStateRelevant = StopStateData.IsActiveState(*this);
+		
+		if (!bIsStopStateRelevant)
+		{
+			bStopTransitionTriggered = false;
+			bStopComplete = false;
+			bStopArrived = false;
+			return;
+		}
+
+		FVector CharacterLocation = CharacterOwner.GetActorLocation(); //  CharacterOwner->GetMesh()->GetComponentLocation()
+
+		//DrawDebugSphere(GetSkelMeshComponent()->GetWorld(), CharacterLocation, 8.f, 16, FColor::Orange);
+
+		FVector CharacterMovementDirection = CharacterVelocity.GetSafeNormal();
+
+		float Deceleration = CharacterAccelerationMag - CharacterMovement.GetMaxBrakingDeceleration() - (CharacterSpeed * CharacterMovement.BrakingFriction * CharacterMovement.BrakingFrictionFactor);
+		float PredictedDistanceCovered = (-(CharacterSpeed * CharacterSpeed) / (2.f * Deceleration));
+
+		//*****
+		//UE_LOG(LogLocomotion, Log, TEXT("***AI_DM: Velocity: %f, Deceleration: %f, Friction: %f"), CharacterSpeed, -Deceleration, CharacterSpeed * CharacterMovement.BrakingFriction * CharacterMovement.BrakingFrictionFactor);
+		//UE_LOG(LogLocomotion, Log, TEXT("Desired Stop Animation Position: %f"), DesiredStopAnimPosition);
+		//UE_LOG(LogLocomotion, Log, TEXT("Virtual Stop Animation Framerate: %f"), (DesiredStopAnimPosition - StopAnimPosition)*30.f);
+
+		//*****
+		//UE_LOG(LogLocomotion, Log, TEXT("Predicted Distance To Stop: %f"), PredictedDistanceCovered);
+
+		if (!bStopTransitionTriggered && bIsStopStateRelevant)
+		{
+			StopActualDistanceFromMarker = 0.0f;
+			StopAnimDistanceTraveled = 0.f;
+			StopAnimTimeElapsed = 0.f;
+			
+			// ?????
+			// For now without walk anims
+			StopAnimation = JogStop;
+
+			//DrawDebugSphere(GetSkelMeshComponent()->GetWorld(), StopAnimLocation, 8.f, 16, FColor::Yellow);
+
+			StopAnimPosition = 0.f;
+
+			bStopTransitionTriggered = true;
+
+			return;
+		}
+
+		// Predict Stop?
+
+		// Find best place
+
+		// Adjust animation
+
+		if (StopAnimation)
+		{
+			//StopAnimTimeElapsed += DeltaTime;
+			if (!bStopArrived && (DistanceCurve.IsValid(StopAnimation)))
+			{
+				StopAnimLocation = CharacterLocation + CharacterMovementDirection * PredictedDistanceCovered;
+
+				float DistanceToStop = (StopAnimLocation - CharacterLocation).Size();
+
+				if (DistanceToStop < 1.f)
+				{
+					DistanceToStop = 1.f;
+					bStopArrived = true;
+				}
+
+				float DesiredStopAnimPosition = DistanceCurve.GetAnimPositionFromDistance(StopAnimation, -DistanceToStop);
+
+				//*****
+				//UE_LOG(LogLocomotion, Log, TEXT("Distance to Stop: %f"), -DistanceToStop);
+				//UE_LOG(LogLocomotion, Log, TEXT("Desired Stop Animation Position: %f"), DesiredStopAnimPosition);
+				//UE_LOG(LogLocomotion, Log, TEXT("Virtual Stop Animation Framerate: %f"), (DesiredStopAnimPosition - StopAnimPosition)*30.f);
+
+				StopAnimPosition = DesiredStopAnimPosition;
+			}
+			else
+			{
+				StopAnimPosition += DeltaTime;
+
+				if (StopAnimPosition >= StopAnimation->SequenceLength - StopTransitionData.GetCrossfadeDuration(*this))
+				{
+					bStopComplete = true;
+				}
+			}
+		}
+
+	}
+}
+
+void UAnimInstance_DistanceMatching::CalculateStrideWarpingValues(const float& InRateWarpCurveAlpha, const float& InOverallRate)
+{
+	float ClampedRateWarpAlpha = FMath::Clamp(InRateWarpCurveAlpha, 0.f, 1.f);
+	if (FMath::IsNearlyEqual(ClampedRateWarpAlpha, 1.f))
+	{
+		PlayRateValue = InOverallRate;
+		StrideWarpingValue = 1.f;
+	}
+	else if (FMath::IsNearlyEqual(ClampedRateWarpAlpha, 0.f))
+	{
+		PlayRateValue = 1.f;
+		StrideWarpingValue = InOverallRate;
+	}
+	else
+	{
+		float AlphaRatio = (ClampedRateWarpAlpha) / (1.f - ClampedRateWarpAlpha);
+		float a = AlphaRatio - 1.f;
+		float radical = FMath::Sqrt((a * a) + (4.f * AlphaRatio * InOverallRate));
+		float x1 = (-a + radical) / 2.f;
+		float x2 = (-a - radical) / 2.f;
+
+		PlayRateValue = FMath::Max(x1, x2);
+		StrideWarpingValue = -((1.f - PlayRateValue - AlphaRatio) / AlphaRatio);
+	}
+}
diff --git a/Engine/Plugins/Animation/PUMA/Source/Runtime/Public/AnimInstances/AnimInstance_DistanceMatching.h b/Engine/Plugins/Animation/PUMA/Source/Runtime/Public/AnimInstances/AnimInstance_DistanceMatching.h
new file mode 100644
--- /dev/null
+++ b/Engine/Plugins/Animation/PUMA/Source/Runtime/Public/AnimInstances/AnimInstance_DistanceMatching.h
@@ -0,0 +1,227 @@
+// Copyright Epic Games, Inc. All Rights Reserved.
+
+#pragma once
+
+#include "Animation/AnimInstance.h"
+#include "Animation/CachedAnimData.h"
+
+#include "GameFramework/Character.h"
+#include "GameFramework/CharacterMovementComponent.h"
+
+#include "AnimInstance_DistanceMatching.generated.h"
+
+UCLASS(transient, Blueprintable, hideCategories=AnimInstance, BlueprintType, meta=(BlueprintThreadSafe), Within=SkeletalMeshComponent)
+class PUMARUNTIME_API UAnimInstance_DistanceMatching : public UAnimInstance
+{
+	GENERATED_UCLASS_BODY()
+
+protected:
+	// Begin UAnimInstance interface
+	virtual void NativeBeginPlay() override;
+	virtual void NativeUpdateAnimation(float DeltaSeconds) override;
+	virtual void NativePostEvaluateAnimation() override;
+	// End UAnimInstance interface
+
+public:
+
+	UPROPERTY(EditAnywhere, BlueprintReadOnly, Transient, Category = Locomotion)
+	FVector CharacterVelocity{ FVector::ZeroVector };
+
+	UPROPERTY(EditAnywhere, BlueprintReadOnly, Transient, Category = Locomotion)
+	float CharacterSpeed{ 0.f };
+
+	UPROPERTY(EditAnywhere, BlueprintReadOnly, Transient, Category = Locomotion)
+	float CharacterSpeed2D{ 0.f };
+
+	UPROPERTY(EditAnywhere, BlueprintReadOnly, Transient, Category = Locomotion)
+	float CharacterSpeedZ{ 0.f };
+
+	UPROPERTY(EditAnywhere, BlueprintReadOnly, Category = Locomotion, Meta = (DisplayName = "Min Character Speed Mag Threshold"))
+	float MinCharacterSpeedThreshold{ 25.f };
+
+	UPROPERTY(EditAnywhere, BlueprintReadOnly, Transient, Category = Locomotion)
+	FVector CharacterAcceleration{ FVector::ZeroVector };
+
+	UPROPERTY(EditAnywhere, BlueprintReadOnly, Transient, Category = Locomotion)
+	float CharacterAccelerationMag{ 0.f };
+
+	UPROPERTY(EditAnywhere, BlueprintReadOnly, Transient, Category = Locomotion)
+	float CharacterAccelerationMag2D{ 0.f };
+
+	UPROPERTY(EditAnywhere, BlueprintReadOnly, Transient, Category = Locomotion)
+	float CharacterAccelerationMagZ{ 0.f };
+
+	UPROPERTY(EditAnywhere, BlueprintReadOnly, Category = Locomotion, Meta = (DisplayName = "Min Character Acceleration Mag Threshold"))
+	float MinCharacterAccelerationMagThreshold{ 25.f };
+
+	UPROPERTY(BlueprintReadOnly, Transient, Category = Locomotion)
+	float MaxCharacterSpeed;
+
+	UPROPERTY(EditAnywhere, BlueprintReadOnly, Category = "Locomotion|Move", Meta = (DisplayName = "Walk Speed Ideal"))
+	float WalkSpeedIdeal{ 150.f };
+
+	UPROPERTY(EditAnywhere, BlueprintReadOnly, Category = "Locomotion|Move", Meta = (DisplayName = "Walk Speed Max"))
+	float WalkSpeedMax{ 200.f };
+
+	UPROPERTY(EditAnywhere, BlueprintReadWrite, Category = "Locomotion|Move")
+	FCachedAnimStateData WalkStateData;
+
+	UPROPERTY(EditAnywhere, BlueprintReadOnly, Category = "Locomotion|Move", Meta = (DisplayName = "Jog Speed Ideal"))
+	float JogSpeedIdeal{ 400.f };
+
+	UPROPERTY(EditAnywhere, BlueprintReadOnly, Category = "Locomotion|Move", Meta = (DisplayName = "Jog Speed Max"))
+	float JogSpeedMax{ 500.f };
+
+	UPROPERTY(EditAnywhere, BlueprintReadWrite, Category = "Locomotion|Move")
+	FCachedAnimStateData JogStateData;
+
+	UPROPERTY(EditAnywhere, BlueprintReadOnly, Category = "Locomotion|Move", Meta = (DisplayName = "Sprint Speed Ideal"))
+	float SprintSpeedIdeal{ 600.f };
+
+	UPROPERTY(EditAnywhere, BlueprintReadOnly, Category = "Locomotion|Move", Meta = (DisplayName = "Sprint Speed Max"))
+	float SprintSpeedMax{ 750.f };
+
+	UPROPERTY(EditAnywhere, BlueprintReadWrite, Category = "Locomotion|Move")
+	FCachedAnimStateData SprintStateData;
+
+	UPROPERTY(EditAnywhere, BlueprintReadWrite, Category = "Locomotion|DistanceCurve")
+	FDistanceCurve DistanceCurve;
+
+	UPROPERTY(EditAnywhere, BlueprintReadWrite, Category = "Locomotion|Start")
+	FCachedAnimStateData StartStateData;
+
+	UPROPERTY(EditAnywhere, BlueprintReadWrite, Category = "Locomotion|Start")
+	FCachedAnimTransitionData StartTransitionData;
+
+	UPROPERTY(EditAnywhere, BlueprintReadWrite, Category = "Locomotion|Start")
+	UAnimSequence* WalkStart;
+
+	UPROPERTY(EditAnywhere, BlueprintReadWrite, Category = "Locomotion|Start")
+	UAnimSequence* JogStart;
+
+	UPROPERTY(BlueprintReadOnly, Transient, Category = "Locomotion|Start")
+	UAnimSequence* StartAnimation;
+
+	UPROPERTY(BlueprintReadOnly, Transient, Category = "Locomotion|Start")
+	uint8 bStartTransitionTriggered : 1;
+
+	UPROPERTY(BlueprintReadOnly, Transient, Category = "Locomotion|Start")
+	uint8 bStartEarlyOut : 1;
+
+	UPROPERTY(BlueprintReadOnly, Transient, Category = "Locomotion|Start")
+	float StartAnimPosition;
+
+	UPROPERTY(Transient)
+	float StartAnimDistanceTraveled;
+
+	UPROPERTY(Transient)
+	float StartActualDistanceFromMarker;
+
+	UPROPERTY(Transient)
+	float StartAnimTimeElapsed;
+
+	UPROPERTY(Transient)
+	FVector StartAnimLocation;
+
+	UPROPERTY(BlueprintReadWrite, Transient, Category = "Locomotion|Start")
+	bool bStartComplete;
+
+	UPROPERTY(BlueprintReadWrite, Transient, Category = "Locomotion|Start")
+	bool bPlayStart;
+
+	UPROPERTY(BlueprintReadWrite, Transient, Category = "Locomotion|Start")
+	bool bWalkToStart;
+
+	UPROPERTY(EditAnywhere, BlueprintReadWrite, Category = "Locomotion|Stop")
+	FCachedAnimStateData StopStateData;
+
+	UPROPERTY(EditAnywhere, BlueprintReadWrite, Category = "Locomotion|Stop")
+	FCachedAnimTransitionData StopTransitionData;
+
+	UPROPERTY(EditAnywhere, BlueprintReadWrite, Category = "Locomotion|Stop")
+	UAnimSequence* WalkStop;
+
+	UPROPERTY(EditAnywhere, BlueprintReadWrite, Category = "Locomotion|Stop")
+	UAnimSequence* JogStop;
+
+	UPROPERTY(BlueprintReadOnly, Transient, Category = "Locomotion|Stop")
+	UAnimSequence* StopAnimation;
+
+	UPROPERTY(BlueprintReadOnly, Transient, Category = "Locomotion|Stop")
+	uint8 bStopTransitionTriggered : 1;
+
+	UPROPERTY(BlueprintReadOnly, Transient, Category = "Locomotion|Stop")
+	uint8 bStopEarlyOut : 1;
+
+	UPROPERTY(BlueprintReadOnly, Transient, Category = "Locomotion|Stop")
+	float StopAnimPosition;
+
+	UPROPERTY(Transient)
+	float StopAnimDistanceTraveled;
+
+	UPROPERTY(Transient)
+	float StopActualDistanceFromMarker;
+
+	UPROPERTY(Transient)
+	float StopAnimTimeElapsed;
+
+	UPROPERTY(Transient)
+	FVector StopAnimLocation;
+
+	UPROPERTY(BlueprintReadWrite, Transient, Category = "Locomotion|Stop")
+	bool bStopArrived;
+
+	UPROPERTY(BlueprintReadWrite, Transient, Category = "Locomotion|Stop")
+	bool bStopComplete;
+
+	UPROPERTY(EditAnywhere, BlueprintReadWrite, Category = "Locomotion|PlayRateAndStrideWarping")
+	FCachedFloatCurve PlayRateStrideWarpAlphaCurve;
+
+	UPROPERTY(EditAnywhere, BlueprintReadWrite, Category = "Locomotion|PlayRateAndStrideWarping")
+	float PlayRateStrideWarpAlpha{ 0.5f };
+
+	UPROPERTY(BlueprintReadOnly, Transient, Category = "Locomotion|PlayRateAndStrideWarping")
+	float PlayRateValue;
+
+	UPROPERTY(BlueprintReadOnly, Transient, Category = "Locomotion|PlayRateAndStrideWarping")
+	float StrideWarpingValue;
+
+	UPROPERTY(BlueprintReadOnly, Transient, Category = "Locomotion|Slope")
+	float SlopeAngle;
+
+	UPROPERTY(BlueprintReadOnly, Transient, Category = "Locomotion|Slope")
+	float SlopeWarpingAlpha;
+
+	UPROPERTY(BlueprintReadOnly, Transient, Category = "Locomotion|Slope")
+	float SlopeWarpingAdjustmentAlpha;
+
+	UPROPERTY(EditAnywhere, BlueprintReadWrite, Category = "Locomotion|Slope")
+	float MaxCharacterSpeedForSlopeWarpingDamping{ 275.f };
+
+	UPROPERTY(EditAnywhere, BlueprintReadWrite, Category = "Locomotion|Slope", meta = (ClampMin = "0.0", ClampMax = "1.0", UIMin = "0.0", UIMax = "1.0"))
+	float MinSlopeWarpingAlpha{ 0.5f };
+
+	UPROPERTY(EditAnywhere, BlueprintReadWrite, Category = "Locomotion|Slope", meta = (ClampMin = "0.0", ClampMax = "1.0", UIMin = "0.0", UIMax = "1.0"))
+	float MaxSlopeWarpingAlpha{ 1.f };
+
+	UPROPERTY(EditAnywhere, BlueprintReadWrite, Category = "Locomotion|Slope")
+	float MaxCharacterSpeedForSlopeAngleAdjustmentRange{ 400.f };
+
+	UPROPERTY(EditAnywhere, BlueprintReadWrite, Category = "Locomotion|Slope", meta = (ClampMin = "0.0", ClampMax = "1.0", UIMin = "0.0", UIMax = "1.0"))
+	float MinSlopeAngleAdjustmentSpeed{ 45.f };
+
+	UPROPERTY(EditAnywhere, BlueprintReadWrite, Category = "Locomotion|Slope", meta = (ClampMin = "0.0", ClampMax = "1.0", UIMin = "0.0", UIMax = "1.0"))
+	float MaxSlopeAngleAdjustmentSpeed{ 90.f };
+
+private:
+
+	void UpdateCharacterLocomotionProperties(const ACharacter& CharacterOwner, const UCharacterMovementComponent& CharacterMovement, float DeltaSeconds);
+
+	// Distance-matched transitions
+	void UpdateStartTransition(const ACharacter& CharacterOwner, const UCharacterMovementComponent& CharacterMovement, float DeltaSeconds);
+	void UpdateStopTransition(const ACharacter& CharacterOwner, const UCharacterMovementComponent& CharacterMovement, float DeltaSeconds);
+
+	void CalculateStrideWarpingValues(const float& InRateWarpCurveAlpha, const float& InOverallRate);
+
+	FTransform IKFootRootTransformLastFrame;
+};
