SequenceLength access deprecation fix-up

diff --git a/Engine/Plugins/Animation/PUMA/Source/Runtime/Private/AnimInstances/AnimInstance_DistanceMatching.cpp b/Engine/Plugins/Animation/PUMA/Source/Runtime/Private/AnimInstances/AnimInstance_DistanceMatching.cpp
--- a/Engine/Plugins/Animation/PUMA/Source/Runtime/Private/AnimInstances/AnimInstance_DistanceMatching.cpp
+++ b/Engine/Plugins/Animation/PUMA/Source/Runtime/Private/AnimInstances/AnimInstance_DistanceMatching.cpp
@@ -185,7 +185,7 @@ void UAnimInstance_DistanceMatching::UpdateStartTransition(const ACharacter& Cha
 				FVector FrameMovement = (CharacterVelocity * DeltaTime) + (CharacterAcceleration * DeltaTime * DeltaTime);
 									
 				float DesiredStartAnimDistance = DistanceCurve.GetValueAtPosition(StartAnimation, StartAnimPosition) + FrameMovement.Size2D();
-				float DesiredStartAnimPosition = FMath::Clamp(DistanceCurve.GetAnimPositionFromDistance(StartAnimation, DesiredStartAnimDistance), StartAnimPosition, StartAnimation->SequenceLength);
+				float DesiredStartAnimPosition = FMath::Clamp(DistanceCurve.GetAnimPositionFromDistance(StartAnimation, DesiredStartAnimDistance), StartAnimPosition, StartAnimation->GetPlayLength());
 				float DesiredStartAnimOverallRate = FMath::Clamp((DesiredStartAnimPosition - StartAnimPosition) / DeltaTime, 0.75f, 1.5f);
 				CalculateStrideWarpingValues(PlayRateStrideWarpAlpha, DesiredStartAnimOverallRate);
 				StartAnimPosition += DeltaTime * PlayRateValue;
@@ -298,7 +298,7 @@ void UAnimInstance_DistanceMatching::UpdateStopTransition(const ACharacter& Char
 			{
 				StopAnimPosition += DeltaTime;
 
-				if (StopAnimPosition >= StopAnimation->SequenceLength - StopTransitionData.GetCrossfadeDuration(*this))
+				if (StopAnimPosition >= StopAnimation->GetPlayLength() - StopTransitionData.GetCrossfadeDuration(*this))
 				{
 					bStopComplete = true;
 				}
diff --git a/Engine/Plugins/Animation/PoseSearch/Source/Runtime/Private/PoseSearch.cpp b/Engine/Plugins/Animation/PoseSearch/Source/Runtime/Private/PoseSearch.cpp
--- a/Engine/Plugins/Animation/PoseSearch/Source/Runtime/Private/PoseSearch.cpp
+++ b/Engine/Plugins/Animation/PoseSearch/Source/Runtime/Private/PoseSearch.cpp
@@ -296,7 +296,7 @@ void PoseSearchBuildIndex(const UAnimSequenceBase& AnimSequence, const UPoseSear
 	ExtractionCtx.bExtractRootMotion = true;
 
 	double CurrTime = SearchConfig.FrameSamplingRange.GetLowerBoundValue();
-	const double EndTime = FMath::Min(AnimSequence.SequenceLength, SearchConfig.FrameSamplingRange.GetUpperBoundValue());
+	const double EndTime = FMath::Min(AnimSequence.GetPlayLength(), SearchConfig.FrameSamplingRange.GetUpperBoundValue());
 	const double DeltaTime = 1.0 / SearchConfig.SampleRate;
 
 	const int32 NumPoses = FMath::FloorToInt((EndTime - CurrTime) / DeltaTime);
diff --git a/Engine/Source/Runtime/Engine/Private/Animation/AnimSequence.cpp b/Engine/Source/Runtime/Engine/Private/Animation/AnimSequence.cpp
--- a/Engine/Source/Runtime/Engine/Private/Animation/AnimSequence.cpp
+++ b/Engine/Source/Runtime/Engine/Private/Animation/AnimSequence.cpp
@@ -4297,7 +4297,7 @@ void UAnimSequence::BakeTrackCurvesToRawAnimationTracks(TArray<FRawAnimSequenceT
 
 
 		// NumFrames can't be zero (filtered earlier)
-		const float Interval = GetIntervalPerKey(NumFrames, SequenceLength);
+		const float Interval = GetIntervalPerKey(NumFrames, GetPlayLength());
 
 		for (const FTransformCurve& Curve : RawCurveData.TransformCurves)
 		{
